Build nvTriStrip and ForsythTriOptimizer with UE clang toolchain (link with libstdc++)

Remove stdc++ dependency for MeshBuilderCommon

diff --git a/Engine/Source/Developer/MeshBuilderCommon/MeshBuilderCommon.Build.cs b/Engine/Source/Developer/MeshBuilderCommon/MeshBuilderCommon.Build.cs
--- a/Engine/Source/Developer/MeshBuilderCommon/MeshBuilderCommon.Build.cs
+++ b/Engine/Source/Developer/MeshBuilderCommon/MeshBuilderCommon.Build.cs
@@ -20,11 +20,6 @@ namespace UnrealBuildTool.Rules
                 "ForsythTriOptimizer",
                 "nvTessLib"
                 );
-
-            if (Target.IsInPlatformGroup(UnrealPlatformGroup.Unix))
-            {
-                PublicSystemLibraries.Add("stdc++");    // can be fixed, see UE-70769
-            }
        }
 	}
 }
