Always define _SILENCE_STDEXT_HASH_DEPRECATION_WARNINGS in this file as it is going to be removed from being set in UBT

#rnx

diff --git a/Engine/Source/Developer/MeshBuilderCommon/Private/ThirdPartyBuildOptimizationHelper.cpp b/Engine/Source/Developer/MeshBuilderCommon/Private/ThirdPartyBuildOptimizationHelper.cpp
--- a/Engine/Source/Developer/MeshBuilderCommon/Private/ThirdPartyBuildOptimizationHelper.cpp
+++ b/Engine/Source/Developer/MeshBuilderCommon/Private/ThirdPartyBuildOptimizationHelper.cpp
@@ -6,7 +6,7 @@
 #include "NvTriStrip.h"
 #include "forsythtriangleorderoptimizer.h"
 
-#ifdef __clang__
+#ifndef _SILENCE_STDEXT_HASH_DEPRECATION_WARNINGS
 #define _SILENCE_STDEXT_HASH_DEPRECATION_WARNINGS 1
 #endif
 
